Fixed monolithic games not compiling (CIS)

diff --git a/Engine/Source/Editor/EditorLiveStreaming/EditorLiveStreaming.cpp b/Engine/Source/Editor/EditorLiveStreaming/EditorLiveStreaming.cpp
--- a/Engine/Source/Editor/EditorLiveStreaming/EditorLiveStreaming.cpp
+++ b/Engine/Source/Editor/EditorLiveStreaming/EditorLiveStreaming.cpp
@@ -9,7 +9,7 @@
 #include "Editor/MainFrame/Public/Interfaces/IMainFrameModule.h"
 #include "Runtime/GameLiveStreaming/Public/IGameLiveStreaming.h"
 
-IMPLEMENT_MODULE( FEditorLiveStreaming, "EditorLiveStreaming" );
+IMPLEMENT_MODULE( FEditorLiveStreaming, EditorLiveStreaming );
 
 #define LOCTEXT_NAMESPACE "EditorLiveStreaming"
 
diff --git a/Engine/Source/Runtime/GameLiveStreaming/GameLiveStreaming.cpp b/Engine/Source/Runtime/GameLiveStreaming/GameLiveStreaming.cpp
--- a/Engine/Source/Runtime/GameLiveStreaming/GameLiveStreaming.cpp
+++ b/Engine/Source/Runtime/GameLiveStreaming/GameLiveStreaming.cpp
@@ -8,7 +8,7 @@
 #include "RenderCore.h"
 #include "RHIStaticStates.h"
 
-IMPLEMENT_MODULE( FGameLiveStreaming, "GameLiveStreaming" );
+IMPLEMENT_MODULE( FGameLiveStreaming, GameLiveStreaming );
 
 #define LOCTEXT_NAMESPACE "GameLiveStreaming"
 
